Static analysis fix

diff --git a/Engine/Plugins/Runtime/PacketHandlers/DTLSHandlerComponent/Source/Private/DTLSContext.cpp b/Engine/Plugins/Runtime/PacketHandlers/DTLSHandlerComponent/Source/Private/DTLSContext.cpp
--- a/Engine/Plugins/Runtime/PacketHandlers/DTLSHandlerComponent/Source/Private/DTLSContext.cpp
+++ b/Engine/Plugins/Runtime/PacketHandlers/DTLSHandlerComponent/Source/Private/DTLSContext.cpp
@@ -384,7 +384,7 @@ bool FDTLSContext::Initialize(const int32 MaxPacketSize, const FString& CertId,
 	SSL_set_mtu(SSLPtr, MaxPacketSize);
 
 	MTUFilterMethod = BIO_meth_new(BIO_TYPE_FILTER, "MTUFilter");
-	if (!SSLPtr)
+	if (!MTUFilterMethod)
 	{
 		UE_LOG(LogDTLSHandler, Error, TEXT("Failed to create SSL MTU filter."));
 		return false;
