NetCore CIS fixes.

diff --git a/Engine/Plugins/Runtime/PacketHandlers/DTLSHandlerComponent/Source/DTLSHandlerComponent.Build.cs b/Engine/Plugins/Runtime/PacketHandlers/DTLSHandlerComponent/Source/DTLSHandlerComponent.Build.cs
--- a/Engine/Plugins/Runtime/PacketHandlers/DTLSHandlerComponent/Source/DTLSHandlerComponent.Build.cs
+++ b/Engine/Plugins/Runtime/PacketHandlers/DTLSHandlerComponent/Source/DTLSHandlerComponent.Build.cs
@@ -12,6 +12,7 @@ public class DTLSHandlerComponent : ModuleRules
 			{
 				"Core",
                 "CoreUObject",
+		"NetCore",
                 "PacketHandler",
                 "Engine",
                 "SSL",
