Experimental DTLS packet handler
(requires OpenSSL 1.1.0 or higher)

diff --git a/Engine/Plugins/Runtime/PacketHandlers/DTLSHandlerComponent/DTLSHandlerComponent.uplugin b/Engine/Plugins/Runtime/PacketHandlers/DTLSHandlerComponent/DTLSHandlerComponent.uplugin
new file mode 100644
--- /dev/null
+++ b/Engine/Plugins/Runtime/PacketHandlers/DTLSHandlerComponent/DTLSHandlerComponent.uplugin
@@ -0,0 +1,33 @@
+{
+	"FileVersion" : 3,
+	"Version" : 1,
+	"VersionName" : "1.0",
+	"FriendlyName" : "DTLS network packet handler",
+	"Description" : "Provides a packet handler component to do DTLS encryption and decryption.",
+	"Category" : "Misc",
+	"CreatedBy" : "Epic Games, Inc.",
+	"CreatedByURL" : "http://epicgames.com",
+	"DocsURL" : "",
+	"MarketplaceURL" : "",
+	"SupportURL" : "",
+	"EnabledByDefault" : false,
+	"CanContainContent" : false,
+	"IsExperimentalVersion" : true,
+	"Installed" : false,
+	"Modules" :
+	[
+		{
+			"Name" : "DTLSHandlerComponent",
+			"Type" : "Runtime",
+			"LoadingPhase" : "Default",
+			"WhitelistPlatforms":
+			[
+				"IOS",
+				"Mac",
+				"Win32",
+				"Win64",
+				"Linux"
+			]
+		}
+	]
+}
\ No newline at end of file
diff --git a/Engine/Plugins/Runtime/PacketHandlers/DTLSHandlerComponent/Source/DTLSHandlerComponent.Build.cs b/Engine/Plugins/Runtime/PacketHandlers/DTLSHandlerComponent/Source/DTLSHandlerComponent.Build.cs
new file mode 100644
--- /dev/null
+++ b/Engine/Plugins/Runtime/PacketHandlers/DTLSHandlerComponent/Source/DTLSHandlerComponent.Build.cs
@@ -0,0 +1,23 @@
+// Copyright 1998-2019 Epic Games, Inc. All Rights Reserved.
+
+using UnrealBuildTool;
+using System.IO;
+
+public class DTLSHandlerComponent : ModuleRules
+{
+    public DTLSHandlerComponent(ReadOnlyTargetRules Target) : base(Target)
+	{
+        PrivateDependencyModuleNames.AddRange(
+			new string[]
+			{
+				"Core",
+                "CoreUObject",
+                "PacketHandler",
+                "Engine",
+                "SSL",
+			}
+		);
+
+        AddEngineThirdPartyPrivateStaticDependencies(Target, "OpenSSL");
+    }
+}
\ No newline at end of file
diff --git a/Engine/Plugins/Runtime/PacketHandlers/DTLSHandlerComponent/Source/Private/DTLSCertStore.cpp b/Engine/Plugins/Runtime/PacketHandlers/DTLSHandlerComponent/Source/Private/DTLSCertStore.cpp
new file mode 100644
--- /dev/null
+++ b/Engine/Plugins/Runtime/PacketHandlers/DTLSHandlerComponent/Source/Private/DTLSCertStore.cpp
@@ -0,0 +1,62 @@
+// Copyright 1998-2019 Epic Games, Inc. All Rights Reserved.
+
+#include "DTLSCertStore.h"
+
+#if !UE_BUILD_SHIPPING
+static TAutoConsoleVariable<int32> CVarDTLSDebugFingerprints(TEXT("DTLS.DebugFingerprints"), 1, TEXT(""));
+#endif
+
+TUniquePtr<FDTLSCertStore> FDTLSCertStore::Instance;
+
+FDTLSCertStore& FDTLSCertStore::Get()
+{
+	if (!Instance.IsValid())
+	{
+		Instance.Reset(new FDTLSCertStore());
+	}
+
+	return *Instance.Get();
+}
+
+TSharedPtr<FDTLSCertificate> FDTLSCertStore::CreateCert()
+{
+	TSharedRef<FDTLSCertificate> Cert = MakeShared<FDTLSCertificate>();
+	if (Cert->GenerateCertificate())
+	{
+		return Cert;
+	}
+
+	UE_LOG(LogDTLSHandler, Error, TEXT("CreateCert: Failed to create certificate"));
+	return nullptr;
+}
+
+TSharedPtr<FDTLSCertificate> FDTLSCertStore::CreateCert(const FString& Identifier)
+{
+	TSharedPtr<FDTLSCertificate> Cert = CreateCert();
+	if (Cert.IsValid() && !Identifier.IsEmpty())
+	{
+		CertMap.Emplace(Identifier, Cert);
+
+#if !UE_BUILD_SHIPPING
+		const bool bDebugFingerprints = CVarDTLSDebugFingerprints.GetValueOnAnyThread();
+
+		if (bDebugFingerprints)
+		{
+			FString DebugFilename = FString::Printf(TEXT("%s%s.bin"), *FPaths::ProjectLogDir(), *FPaths::MakeValidFileName(Identifier));
+			FFileHelper::SaveArrayToFile(Cert->GetFingerprint(), *DebugFilename);
+		}
+#endif
+	}
+
+	return Cert;
+}
+
+TSharedPtr<FDTLSCertificate> FDTLSCertStore::GetCert(const FString& Identifier) const
+{
+	return CertMap.FindRef(Identifier);
+}
+
+bool FDTLSCertStore::RemoveCert(const FString& Identifier)
+{
+	return (CertMap.Remove(Identifier) != 0);
+}
diff --git a/Engine/Plugins/Runtime/PacketHandlers/DTLSHandlerComponent/Source/Private/DTLSCertificate.cpp b/Engine/Plugins/Runtime/PacketHandlers/DTLSHandlerComponent/Source/Private/DTLSCertificate.cpp
new file mode 100644
--- /dev/null
+++ b/Engine/Plugins/Runtime/PacketHandlers/DTLSHandlerComponent/Source/Private/DTLSCertificate.cpp
@@ -0,0 +1,127 @@
+// Copyright 1998-2019 Epic Games, Inc. All Rights Reserved.
+
+#include "DTLSCertificate.h"
+#include "Misc/ScopeExit.h"
+
+#define UI UI_ST
+THIRD_PARTY_INCLUDES_START
+#include <openssl/x509.h>
+THIRD_PARTY_INCLUDES_END
+#undef UI
+
+FDTLSCertificate::FDTLSCertificate()
+	: PKey(nullptr)
+	, Certificate(nullptr)
+{
+}
+
+FDTLSCertificate::~FDTLSCertificate()
+{
+	FreeCertificate();
+}
+
+void FDTLSCertificate::FreeCertificate()
+{
+	if (Certificate)
+	{
+		X509_free(Certificate);
+		Certificate = nullptr;
+	}
+
+	if (PKey)
+	{
+		EVP_PKEY_free(PKey);
+		PKey = nullptr;
+	}
+
+	Fingerprint.Reset();
+}
+
+bool FDTLSCertificate::GenerateCertificate()
+{
+	bool bSuccess = false;
+
+	ON_SCOPE_EXIT
+	{
+		if (!bSuccess)
+		{
+			FreeCertificate();
+		}
+	};
+
+	check(!PKey && !Certificate);
+
+	PKey = EVP_PKEY_new();
+	Certificate = X509_new();
+
+	// will be freed when the pkey is destroyed
+	RSA* RSAKey = RSA_new();
+	
+	{
+		BIGNUM* BigNum = BN_new();
+		BN_set_word(BigNum, RSA_F4);
+		RSA_generate_key_ex(RSAKey, 2048, BigNum, nullptr);
+		BN_free(BigNum);
+	}
+
+	EVP_PKEY_assign_RSA(PKey, RSAKey);
+
+	X509_set_version(Certificate, 2);
+
+	ASN1_INTEGER* SerialNumber = ASN1_INTEGER_new();
+	
+	{
+		BIGNUM* BigNumTemp = BN_new();
+		BN_rand(BigNumTemp, 159, BN_RAND_TOP_ANY, BN_RAND_BOTTOM_ANY);
+		BN_to_ASN1_INTEGER(BigNumTemp, SerialNumber);
+		BN_free(BigNumTemp);
+	}
+
+	int32 ResultCode = X509_set_serialNumber(Certificate, SerialNumber);
+	if (ResultCode == 0)
+	{
+		UE_LOG(LogDTLSHandler, Error, TEXT("Failed to set serial number on cert: %d"), ResultCode);
+		return bSuccess;
+	}
+
+	ASN1_INTEGER_free(SerialNumber);
+
+	X509_gmtime_adj(X509_get_notBefore(Certificate), 0);
+	X509_gmtime_adj(X509_get_notAfter(Certificate), 31536000L);
+
+	X509_set_pubkey(Certificate, PKey);
+
+	X509_NAME* SubjectName = X509_get_subject_name(Certificate);
+
+	X509_NAME_add_entry_by_txt(SubjectName, "C",  MBSTRING_ASC, (unsigned char *)"US", -1, -1, 0);
+	X509_NAME_add_entry_by_txt(SubjectName, "CN", MBSTRING_ASC, (unsigned char *)"Unreal Engine", -1, -1, 0);
+
+	ResultCode = X509_set_issuer_name(Certificate, SubjectName);
+	if (ResultCode == 0)
+	{
+		UE_LOG(LogDTLSHandler, Error, TEXT("Failed to set issuer name"));
+		return bSuccess;
+	}
+
+	ResultCode = X509_sign(Certificate, PKey, EVP_sha256());
+	if (ResultCode <= 0)
+	{
+		UE_LOG(LogDTLSHandler, Error, TEXT("Failed to sign cert: %d"), ResultCode);
+		return bSuccess;
+	}
+
+	Fingerprint.Reset();
+
+	uint32 HashLen = FDTLSFingerprint::Length;
+
+	ResultCode = X509_digest(Certificate, EVP_sha256(), Fingerprint.Data, &HashLen);
+	if (ResultCode <= 0)
+	{
+		UE_LOG(LogDTLSHandler, Error, TEXT("Failed to hash cert: %d"), ResultCode);
+		return bSuccess;
+	}
+
+	bSuccess = true;
+
+	return bSuccess;
+}
\ No newline at end of file
diff --git a/Engine/Plugins/Runtime/PacketHandlers/DTLSHandlerComponent/Source/Private/DTLSContext.cpp b/Engine/Plugins/Runtime/PacketHandlers/DTLSHandlerComponent/Source/Private/DTLSContext.cpp
new file mode 100644
--- /dev/null
+++ b/Engine/Plugins/Runtime/PacketHandlers/DTLSHandlerComponent/Source/Private/DTLSContext.cpp
@@ -0,0 +1,420 @@
+// Copyright 1998-2019 Epic Games, Inc. All Rights Reserved.
+
+#include "DTLSContext.h"
+#include "DTLSHandlerComponent.h"
+#include "DTLSCertStore.h"
+#include "Ssl.h"
+
+#define UI UI_ST
+THIRD_PARTY_INCLUDES_START
+#include <openssl/ssl.h>
+THIRD_PARTY_INCLUDES_END
+#undef UI
+
+namespace DTLSContext
+{
+	static const char* CipherListPSK = "PSK-AES256-GCM-SHA384";
+	static const char* CipherListCert = "HIGH";
+}
+
+const TCHAR* LexToString(EDTLSContextType ContextType)
+{
+	switch (ContextType)
+	{
+	case EDTLSContextType::Server:
+		return TEXT("Server");
+	case EDTLSContextType::Client:
+		return TEXT("Client");
+	default:
+		return TEXT("Unknown");
+	}
+}
+
+static int DTLSCertVerifyCallback(X509_STORE_CTX* Context, void* UserData)
+{
+	FDTLSContext* DTLSContext = static_cast<FDTLSContext*>(UserData);
+
+	X509* RemoteCert = X509_STORE_CTX_get0_cert(Context);
+	if (RemoteCert)
+	{
+		int32 CmpTimeResult = X509_cmp_current_time(X509_get_notBefore(RemoteCert));
+		if (CmpTimeResult >= 0)
+		{
+			UE_LOG(LogDTLSHandler, Error, TEXT("DTLSCertVerifyCallback: Remote certificate times are not valid."));
+			return 0;
+		}
+
+		CmpTimeResult = X509_cmp_current_time(X509_get_notAfter(RemoteCert));
+		if (CmpTimeResult <= 0)
+		{
+			UE_LOG(LogDTLSHandler, Error, TEXT("DTLSCertVerifyCallback: Remote certificate times are not valid."));
+			return 0;
+		}
+
+		// check subject name == issuer name (self-signed)
+		X509_NAME* IssuerName = X509_get_issuer_name(RemoteCert);
+		X509_NAME* SubjectName = X509_get_subject_name(RemoteCert);
+
+		if (X509_NAME_cmp(IssuerName, SubjectName) != 0)
+		{
+			UE_LOG(LogDTLSHandler, Error, TEXT("DTLSCertVerifyCallback: Remote certificate is not self signed."));
+			return 0;
+		}
+
+		// only checking in clients for now
+		if (DTLSContext && (DTLSContext->GetContextType() == EDTLSContextType::Client))
+		{
+			// check against expected fingerprint
+			FDTLSFingerprint Fingerprint;
+			uint32 HashLen = FDTLSFingerprint::Length;
+
+			int32 ResultCode = X509_digest(RemoteCert, EVP_sha256(), Fingerprint.Data, &HashLen);
+			if (ResultCode <= 0)
+			{
+				UE_LOG(LogDTLSHandler, Error, TEXT("DTLSCertVerifyCallback: Unable to compute fingerprint"));
+				return 0;
+			}
+
+			if (SSL* SSLPtr = static_cast<SSL*>(X509_STORE_CTX_get_ex_data(Context, SSL_get_ex_data_X509_STORE_CTX_idx())))
+			{
+				if (const FDTLSHandlerComponent* Handler = static_cast<FDTLSHandlerComponent*>(SSL_get_app_data(SSLPtr)))
+				{
+					if (const FDTLSFingerprint* RemoteFingerprint = Handler->GetRemoteFingerprint())
+					{
+						if (FMemory::Memcmp(Fingerprint.Data, RemoteFingerprint->Data, HashLen) != 0)
+						{
+							UE_LOG(LogDTLSHandler, Error, TEXT("DTLSCertVerifyCallback: Fingerprint validation failure"));
+							return 0;
+						}
+					}
+					else
+					{
+						UE_LOG(LogDTLSHandler, Error, TEXT("DTLSCertVerifyCallback: Unable to retrieve remote fingerprint"));
+						return 0;
+					}
+				}
+				else
+				{
+					UE_LOG(LogDTLSHandler, Error, TEXT("DTLSCertVerifyCallback: Unable to retrieve handler"));
+					return 0;
+				}
+			}
+			else
+			{
+				UE_LOG(LogDTLSHandler, Error, TEXT("DTLSCertVerifyCallback: Unable to retrieve ssl context"));
+				return 0;
+			}
+		}
+	}
+	else
+	{
+		UE_LOG(LogDTLSHandler, Error, TEXT("DTLSCertVerifyCallback: Unable to retrieve remote cert"));
+		return 1;
+	}
+	
+	UE_LOG(LogDTLSHandler, Log, TEXT("DTLSCertVerifyCallback: Verified"));
+	return 1;
+}
+
+static void DTLSInfoCallback(const SSL* SSLPtr, int InfoType, int Value)
+{
+	UE_LOG(LogDTLSHandler, Verbose, TEXT("SSL Info State: %s Alert: %s Desc: %s"), ANSI_TO_TCHAR(SSL_state_string_long(SSLPtr)), ANSI_TO_TCHAR(SSL_alert_type_string_long(Value)), ANSI_TO_TCHAR(SSL_alert_desc_string_long(Value)));
+}
+
+static unsigned int DTLSPSKClientCallback(SSL* SSLPtr, const char* Hint, char* Identity, unsigned int MaxIdentityLength, unsigned char* PSK, unsigned int MaxPSKLength)
+{
+	if (FDTLSHandlerComponent* Handler = static_cast<FDTLSHandlerComponent*>(SSL_get_app_data(SSLPtr)))
+	{
+		if (const FDTLSPreSharedKey* PreSharedKey = Handler->GetPreSharedKey())
+		{
+			const FString& PSKIdentity = PreSharedKey->GetIdentity();
+			if (!PSKIdentity.IsEmpty())
+			{
+				auto IdentityAnsi = StringCast<ANSICHAR>(*PSKIdentity);
+				int32 IdentityLen = FCStringAnsi::Strlen(IdentityAnsi.Get());
+
+				check(IdentityLen <= (int32)MaxIdentityLength);
+				FCStringAnsi::Strncpy(Identity, IdentityAnsi.Get(), IdentityLen);
+			}
+
+			TArrayView<const uint8> KeyView = PreSharedKey->GetKey();
+
+			check(KeyView.Num() <= (int32)MaxPSKLength);
+			FMemory::Memcpy(PSK, KeyView.GetData(), KeyView.Num());
+		
+			return KeyView.Num();
+		}
+		else
+		{
+			UE_LOG(LogDTLSHandler, Error, TEXT("DTLSPSKClientCallback:  Invalid key"));
+		}
+	}
+	else
+	{
+		UE_LOG(LogDTLSHandler, Error, TEXT("DTLSPSKClientCallback:  Invalid owner"));
+	}
+
+	return 0;
+}
+
+static unsigned int DTLSPSKServerCallback(SSL* SSLPtr, const char* Identity, unsigned char* PSK, unsigned int MaxPSKLength)
+{
+	if (FDTLSHandlerComponent* Handler = static_cast<FDTLSHandlerComponent*>(SSL_get_app_data(SSLPtr)))
+	{
+		if (const FDTLSPreSharedKey* PreSharedKey = Handler->GetPreSharedKey())
+		{
+			const FString& PSKIdentity = PreSharedKey->GetIdentity();
+
+			if (!PSKIdentity.IsEmpty() && Identity)
+			{
+				auto IdentityAnsi = StringCast<ANSICHAR>(*PSKIdentity);
+
+				if (FCStringAnsi::Strcmp(Identity, IdentityAnsi.Get()) != 0)
+				{
+					UE_LOG(LogDTLSHandler, Error, TEXT("DTLSPSKServerCallback: Unexpected identity: %s"), ANSI_TO_TCHAR(Identity));
+					return 0;
+				}
+			}
+
+			TArrayView<const uint8> KeyView = PreSharedKey->GetKey();
+
+			check(KeyView.Num() <= (int32)MaxPSKLength);
+			FMemory::Memcpy(PSK, KeyView.GetData(), KeyView.Num());
+
+			return KeyView.Num();
+		}
+		else
+		{
+			UE_LOG(LogDTLSHandler, Error, TEXT("DTLSPSKServerCallback:  Invalid key"));
+		}
+	}
+	else
+	{
+		UE_LOG(LogDTLSHandler, Error, TEXT("DTLSPSKServerCallback:  Invalid owner"));
+	}
+
+	return 0;
+}
+
+int FMTUFilter::Create(BIO* BIOPtr) 
+{
+	FMTUFilter* MTUFilter = new FMTUFilter();
+
+	BIO_set_init(BIOPtr, 1);
+	BIO_set_data(BIOPtr, MTUFilter);
+	BIO_set_flags(BIOPtr, 0);
+
+ 	return 1;
+}
+
+int FMTUFilter::Destroy(BIO* BIOPtr) 
+{
+	if (BIOPtr == nullptr)
+	{
+		return 0;
+	}
+
+	FMTUFilter* MTUFilter = static_cast<FMTUFilter*>(BIO_get_data(BIOPtr));
+	if (MTUFilter != nullptr) 
+	{
+		delete MTUFilter;
+	}
+
+	BIO_set_data(BIOPtr, nullptr);
+	BIO_set_init(BIOPtr, 0);
+	BIO_set_flags(BIOPtr, 0);
+
+	return 1;
+}
+
+int FMTUFilter::Write(BIO* BIOPtr, const char* Data, int DataLen) 
+{
+	check(IsInGameThread());
+
+	int WriteResult = BIO_write(BIO_next(BIOPtr), Data, DataLen);
+
+	FMTUFilter* MTUFilter = static_cast<FMTUFilter*>(BIO_get_data(BIOPtr));
+	if (MTUFilter != nullptr) 
+	{
+		MTUFilter->FragmentSizes.Add(WriteResult);
+	}
+
+	return WriteResult;
+}
+
+long FMTUFilter::Ctrl(BIO* BIOPtr, int CtrlCommand, long, void*)
+{
+	check(IsInGameThread());
+
+	FMTUFilter* MTUFilter = static_cast<FMTUFilter*>(BIO_get_data(BIOPtr));
+
+	switch (CtrlCommand) 
+	{
+	case BIO_CTRL_FLUSH:
+		return 1;
+	case BIO_CTRL_DGRAM_QUERY_MTU:
+		check(0);
+		return 0;
+	case BIO_CTRL_WPENDING:
+		return 0;
+	case BIO_CTRL_PENDING: 
+	{
+		if (MTUFilter == nullptr || MTUFilter->FragmentSizes.Num() == 0)
+		{
+			return 0;
+		}
+
+		long FragmentSize = MTUFilter->FragmentSizes[0];
+		MTUFilter->FragmentSizes.RemoveAt(0);
+		return FragmentSize;
+	}
+	default:
+		break;
+	}
+
+	return 0;
+}
+
+FDTLSContext::FDTLSContext(EDTLSContextType InContextType)
+	: ContextType(InContextType)
+	, SSLContext(nullptr)
+	, SSLPtr(nullptr)
+	, InBIO(nullptr)
+	, OutBIO(nullptr)
+	, MTUFilterMethod(nullptr)
+{
+}
+
+FDTLSContext::~FDTLSContext()
+{
+	if (MTUFilterMethod)
+	{
+		BIO_meth_free(MTUFilterMethod);
+	}
+
+	// bios are freed by SSL_free
+
+	if (SSLPtr)
+	{
+		SSL_free(SSLPtr);
+	}
+
+	if (SSLContext)
+	{
+		SSL_CTX_free(SSLContext);
+	}
+}
+
+bool FDTLSContext::Initialize(const int32 MaxPacketSize, const FString& CertId, FDTLSHandlerComponent* Handler)
+{
+	SSLContext = SSL_CTX_new(DTLS_method());
+	if (!SSLContext)
+	{
+		UE_LOG(LogDTLSHandler, Error, TEXT("Failed to create SSL context."));
+		return false;
+	}
+
+	SSL_CTX_set_min_proto_version(SSLContext, DTLS1_2_VERSION);
+
+	FSslModule::Get().GetCertificateManager().AddCertificatesToSslContext(SSLContext);
+
+	const bool bPreSharedKeys = (CVarPreSharedKeys.GetValueOnAnyThread() != 0);
+
+	const char* CipherList = bPreSharedKeys ? DTLSContext::CipherListPSK : DTLSContext::CipherListCert;
+
+	int32 Result = SSL_CTX_set_cipher_list(SSLContext, CipherList);
+	if (Result == 0) 
+	{
+		UE_LOG(LogDTLSHandler, Error, TEXT("Failed to set cipher list on SSL context: %d"), Result);
+		return false;
+	}
+
+	SSL_CTX_set_app_data(SSLContext, this);
+	SSL_CTX_set_mode(SSLContext, SSL_MODE_AUTO_RETRY);
+	SSL_CTX_set_options(SSLContext, SSL_OP_NO_QUERY_MTU);
+
+	if (bPreSharedKeys)
+	{
+		SSL_CTX_set_psk_client_callback(SSLContext, DTLSPSKClientCallback);
+		SSL_CTX_set_psk_server_callback(SSLContext, DTLSPSKServerCallback);
+	}
+	else
+	{
+		SSL_CTX_set_verify(SSLContext, SSL_VERIFY_PEER | SSL_VERIFY_FAIL_IF_NO_PEER_CERT, nullptr);
+		SSL_CTX_set_cert_verify_callback(SSLContext, DTLSCertVerifyCallback, this);
+
+		if (!CertId.IsEmpty())
+		{
+			Cert = FDTLSCertStore::Get().GetCert(CertId);
+		}
+		else
+		{
+			UE_LOG(LogDTLSHandler, Warning, TEXT("Empty certificate identifier"));
+			Cert = FDTLSCertStore::Get().CreateCert();
+		}
+
+		if (Cert.IsValid())
+		{
+			SSL_CTX_use_certificate(SSLContext, Cert->GetCertificate());
+			SSL_CTX_use_PrivateKey(SSLContext, Cert->GetPKey());
+			
+			Result = SSL_CTX_check_private_key(SSLContext);
+			if (Result != 1)
+			{
+				UE_LOG(LogDTLSHandler, Error, TEXT("Private key check failed: %d"), Result);
+				return false;
+			}
+		}
+		else
+		{
+			UE_LOG(LogDTLSHandler, Error, TEXT("Failed to retrieve certificate"));
+			return false;
+		}
+	}
+
+	SSLPtr = SSL_new(SSLContext);
+	if (!SSLPtr)
+	{
+		UE_LOG(LogDTLSHandler, Error, TEXT("Failed to create SSL session."));
+		return false;
+	}
+
+	SSL_set_app_data(SSLPtr, Handler);
+	SSL_set_info_callback(SSLPtr, DTLSInfoCallback);
+	SSL_set_mtu(SSLPtr, MaxPacketSize);
+
+	MTUFilterMethod = BIO_meth_new(BIO_TYPE_FILTER, "MTUFilter");
+	if (!SSLPtr)
+	{
+		UE_LOG(LogDTLSHandler, Error, TEXT("Failed to create SSL MTU filter."));
+		return false;
+	}
+
+	BIO_meth_set_create(MTUFilterMethod, &FMTUFilter::Create);
+	BIO_meth_set_destroy(MTUFilterMethod, &FMTUFilter::Destroy);
+	BIO_meth_set_write(MTUFilterMethod, &FMTUFilter::Write);
+	BIO_meth_set_ctrl(MTUFilterMethod, &FMTUFilter::Ctrl);
+
+	InBIO = BIO_new(BIO_s_mem());
+	OutBIO = BIO_new(BIO_s_mem());
+	FilterBIO = BIO_new(MTUFilterMethod);
+
+	BIO_push(FilterBIO, OutBIO);
+	SSL_set_bio(SSLPtr, InBIO, FilterBIO);
+
+	if (ContextType == EDTLSContextType::Server)
+	{
+		SSL_set_accept_state(SSLPtr);
+	}
+	else
+	{
+		SSL_set_connect_state(SSLPtr);
+	}
+
+	return true;
+}
+
+bool FDTLSContext::IsHandshakeComplete() const
+{
+	return (SSLPtr != nullptr) && SSL_is_init_finished(SSLPtr);
+}
\ No newline at end of file
diff --git a/Engine/Plugins/Runtime/PacketHandlers/DTLSHandlerComponent/Source/Private/DTLSHandlerComponent.cpp b/Engine/Plugins/Runtime/PacketHandlers/DTLSHandlerComponent/Source/Private/DTLSHandlerComponent.cpp
new file mode 100644
--- /dev/null
+++ b/Engine/Plugins/Runtime/PacketHandlers/DTLSHandlerComponent/Source/Private/DTLSHandlerComponent.cpp
@@ -0,0 +1,398 @@
+// Copyright 1998-2019 Epic Games, Inc. All Rights Reserved.
+
+#include "DTLSHandlerComponent.h"
+#include "Engine/NetConnection.h"
+#include "Ssl.h"
+
+#define UI UI_ST
+THIRD_PARTY_INCLUDES_START
+#include <openssl/ssl.h>
+#include <openssl/dtls1.h>
+THIRD_PARTY_INCLUDES_END
+#undef UI
+
+DEFINE_LOG_CATEGORY(LogDTLSHandler);
+
+IMPLEMENT_MODULE(FDTLSHandlerComponentModule, DTLSHandlerComponent)
+
+TAutoConsoleVariable<int32> CVarPreSharedKeys(TEXT("DTLS.PreSharedKeys"), 1, TEXT("If non-zero, use pre-shared keys, otherwise self-signed certificates will be generated."));
+
+void FDTLSHandlerComponentModule::StartupModule()
+{
+	FPacketHandlerComponentModuleInterface::StartupModule();
+
+	FSslModule& SslModule = FModuleManager::LoadModuleChecked<FSslModule>("SSL");
+	if (!SslModule.GetSslManager().InitializeSsl())
+	{
+		SSL_library_init();
+		SSL_load_error_strings();        
+		OpenSSL_add_all_algorithms();
+	}
+}
+
+void FDTLSHandlerComponentModule::ShutdownModule()
+{
+	FSslModule& SslModule = FModuleManager::LoadModuleChecked<FSslModule>("SSL");
+	SslModule.GetSslManager().ShutdownSsl();
+
+	FPacketHandlerComponentModuleInterface::ShutdownModule();
+}
+
+TSharedPtr<HandlerComponent> FDTLSHandlerComponentModule::CreateComponentInstance(FString& Options)
+{
+	TSharedRef<HandlerComponent> ReturnVal = MakeShared<FDTLSHandlerComponent>();
+
+	return ReturnVal;
+}
+
+FDTLSHandlerComponent::FDTLSHandlerComponent()
+	: FEncryptionComponent(FName(TEXT("DTLSHandlerComponent")))
+	, InternalState(EDTLSHandlerState::Unencrypted)
+{
+}
+
+void FDTLSHandlerComponent::SetEncryptionData(const FEncryptionData& EncryptionData)
+{
+	const bool bPreSharedKeys = (CVarPreSharedKeys.GetValueOnAnyThread() != 0);
+
+	if (bPreSharedKeys)
+	{
+		PreSharedKey.Reset(new FDTLSPreSharedKey());
+		PreSharedKey->SetPreSharedKey(EncryptionData.Key);
+	}
+	else
+	{
+		if (Handler->Mode == Handler::Mode::Server)
+		{
+			CertId = EncryptionData.Identifier;
+		}
+		else
+		{
+			RemoteFingerprint.Reset(new FDTLSFingerprint());
+
+			if (EncryptionData.Fingerprint.Num() == FDTLSFingerprint::Length)
+			{
+				FMemory::Memcpy(RemoteFingerprint->Data, EncryptionData.Fingerprint.GetData(), EncryptionData.Fingerprint.Num());
+			}
+			else
+			{
+				UE_LOG(LogDTLSHandler, Error, TEXT("SetEncryptionData: Invalid fingerprint size: got %d expected %d"), EncryptionData.Fingerprint.Num(), FDTLSFingerprint::Length);
+			}
+		}
+	}
+}
+
+void FDTLSHandlerComponent::EnableEncryption()
+{
+	UE_LOG(LogDTLSHandler, Verbose, TEXT("EnableEncryption"));
+
+	EDTLSContextType ContextType = (Handler->Mode == Handler::Mode::Server) ? EDTLSContextType::Server : EDTLSContextType::Client;
+
+	DTLSContext.Reset(new FDTLSContext(ContextType));
+	if (DTLSContext.IsValid())
+	{
+		check(Handler);
+
+		const int32 MaxPacketSize = (MaxOutgoingBits - GetReservedPacketBits()) / 8;
+
+		if (DTLSContext->Initialize(MaxPacketSize, CertId, this))
+		{
+			InternalState = EDTLSHandlerState::Handshaking;
+		}
+		else
+		{
+			UE_LOG(LogDTLSHandler, Error, TEXT("EnableEncryption: Failed to initialize context."));
+			DTLSContext.Release();
+		}
+	}
+	else
+	{
+		UE_LOG(LogDTLSHandler, Error, TEXT("EnableEncryption: Invalid context."));
+	}
+}
+
+void FDTLSHandlerComponent::DisableEncryption()
+{
+	UE_LOG(LogDTLSHandler, Verbose, TEXT("DisableEncryption"));
+
+	DTLSContext.Release();
+}
+
+bool FDTLSHandlerComponent::IsEncryptionEnabled() const
+{
+	return DTLSContext.IsValid();
+}
+
+void FDTLSHandlerComponent::Initialize()
+{
+	SetActive(true);
+	SetState(Handler::Component::State::Initialized);
+	Initialized();
+}
+
+bool FDTLSHandlerComponent::IsValid() const
+{
+	return true;
+}
+
+void FDTLSHandlerComponent::Incoming(FBitReader& Packet)
+{
+	DECLARE_SCOPE_CYCLE_COUNTER(TEXT("PacketHandler DTLS Decrypt"), STAT_PacketHandler_DTLS_Decrypt, STATGROUP_Net);
+
+	// check encrypted bit
+	if (Packet.ReadBit() != 0)
+	{
+		if (DTLSContext.IsValid())
+		{
+			const int32 HandshakeBit = Packet.ReadBit();
+			const int32 PayloadBytes = Packet.GetBytesLeft();
+
+			check(PayloadBytes > 0);
+			check(PayloadBytes <= sizeof(TempBuffer));
+
+			TempBuffer[PayloadBytes - 1] = 0;
+			Packet.SerializeBits(TempBuffer, Packet.GetBitsLeft());
+
+			if (InternalState == EDTLSHandlerState::Handshaking)
+			{
+				if (HandshakeBit == 1)
+				{
+					// feed the incoming data to the context
+					const int32 BytesWritten = BIO_write(DTLSContext->GetInBIO(), TempBuffer, PayloadBytes);
+					if (BytesWritten != PayloadBytes)
+					{
+						UE_LOG(LogDTLSHandler, Warning, TEXT("Failed to write entire incoming packet to input BIO: %d / %d"), BytesWritten, PayloadBytes);
+						Packet.SetError();
+						return;
+					}
+				}
+				else
+				{
+					UE_LOG(LogDTLSHandler, Warning, TEXT("Ignoring non-handshake packet while handshake is still in progress."));
+					Packet.SetData(nullptr, 0);
+				}
+			}
+			else if (InternalState == EDTLSHandlerState::Encrypted)
+			{
+				if (HandshakeBit == 0)
+				{
+					const int32 BytesWritten = BIO_write(DTLSContext->GetInBIO(), TempBuffer, PayloadBytes);
+					if (BytesWritten == PayloadBytes)
+					{
+						const int32 BytesRead = SSL_read(DTLSContext->GetSSLPtr(), TempBuffer, sizeof(TempBuffer));
+						if (BytesRead > 0)
+						{
+							// Look for the termination bit that was written in Outgoing() to determine the exact bit size.
+							uint8 LastByte = TempBuffer[BytesRead - 1];
+							if (LastByte != 0)
+							{
+								int32 BitSize = (BytesRead * 8) - 1;
+
+								// Bit streaming, starts at the Least Significant Bit, and ends at the MSB.
+								while (!(LastByte & 0x80))
+								{
+									LastByte *= 2;
+									BitSize--;
+								}
+
+								Packet.SetData(TempBuffer, BitSize);
+							}
+							else
+							{
+								UE_LOG(LogTemp, Error, TEXT("DTLS Error"));
+								Packet.SetError();
+							}
+						}
+						else
+						{
+							const int32 ErrorCode = SSL_get_error(DTLSContext->GetSSLPtr(), BytesRead);
+							UE_LOG(LogDTLSHandler, Error, TEXT("SSL_read error: %d"), ErrorCode);
+							Packet.SetError();
+						}
+					}
+					else
+					{
+						UE_LOG(LogDTLSHandler, Error, TEXT("Failed to write entire incoming packet to input BIO: %d / %d"), BytesWritten, PayloadBytes);
+						Packet.SetError();
+					}
+				}
+				else
+				{
+					UE_LOG(LogDTLSHandler, Warning, TEXT("Ignoring handshake packet received after completion."));
+					Packet.SetData(nullptr, 0);
+				}
+			}
+			else
+			{
+				UE_LOG(LogDTLSHandler, Error, TEXT("Attempted to process packet with handler in invalid state"));
+				Packet.SetError();
+			}
+		}
+		else
+		{
+			UE_LOG(LogDTLSHandler, Error, TEXT("Invalid DTLS context."));
+			Packet.SetError();
+		}
+	}
+}
+
+void FDTLSHandlerComponent::UpdateHandshake()
+{
+	DECLARE_SCOPE_CYCLE_COUNTER(TEXT("PacketHandler DTLS Handshake"), STAT_PacketHandler_DTLS_Handshake, STATGROUP_Net);
+
+	if ((InternalState == EDTLSHandlerState::Handshaking) && DTLSContext.IsValid())
+	{
+		SSL* SSLPtr = DTLSContext->GetSSLPtr();
+
+		if (!DTLSContext->IsHandshakeComplete())
+		{
+			const int32 HandshakeResult = SSL_do_handshake(SSLPtr);
+			if (HandshakeResult != 1)
+			{
+				int SSLError = SSL_get_error(SSLPtr, HandshakeResult);
+
+				if (SSLError != SSL_ERROR_NONE && SSLError != SSL_ERROR_WANT_READ && SSLError != SSL_ERROR_WANT_WRITE)
+				{
+					UE_LOG(LogDTLSHandler, Error, TEXT("UpdateHandshake:  Handshaking failed with error: %d"), SSLError);
+					DTLSContext.Reset();
+					return;
+				}
+			}
+
+			int32 Pending = BIO_ctrl_pending(DTLSContext->GetFilterBIO());
+			while (Pending > 0)
+			{
+				check(Pending <= sizeof(TempBuffer));
+
+				const int32 BytesRead = BIO_read(DTLSContext->GetOutBIO(), TempBuffer, Pending);
+				if (BytesRead > 0)
+				{
+					check(BytesRead == Pending);
+					check(BytesRead <= MAX_PACKET_SIZE);
+
+					FBitWriter OutPacket(0, true);
+					OutPacket.WriteBit(1);	// encryption enabled
+					OutPacket.WriteBit(1);	// handshake packet
+					OutPacket.SerializeBits(TempBuffer, BytesRead * 8);
+
+					// SendHandlerPacket is a low level send and is not reliable
+					FOutPacketTraits Traits;
+					Handler->SendHandlerPacket(this, OutPacket, Traits);
+				}
+				else
+				{
+					UE_LOG(LogDTLSHandler, Error, TEXT("BIO_read error: %d"), BytesRead);
+					return;
+				}
+
+				Pending = BIO_ctrl_pending(DTLSContext->GetFilterBIO());
+			}
+		}
+		else
+		{
+			InternalState = EDTLSHandlerState::Encrypted;
+			UE_LOG(LogDTLSHandler, Log, TEXT("UpdateHandshake:  Handshaking completed"));
+		}
+	}
+}
+
+void FDTLSHandlerComponent::Outgoing(FBitWriter& Packet, FOutPacketTraits& Traits)
+{
+	DECLARE_SCOPE_CYCLE_COUNTER(TEXT("PacketHandler DTLS Encrypt"), STAT_PacketHandler_DTLS_Encrypt, STATGROUP_Net);
+
+	const bool bEncryptionEnabled = IsEncryptionEnabled();
+
+	FBitWriter NewPacket(Packet.GetNumBits() + 2, true);
+	NewPacket.WriteBit(bEncryptionEnabled ? 1 : 0);
+
+	if (bEncryptionEnabled)
+	{
+		if (DTLSContext.IsValid())
+		{
+			if (InternalState == EDTLSHandlerState::Encrypted)
+			{
+				NewPacket.WriteBit(0); // handshake bit
+
+				// Termination bit
+				Packet.WriteBit(1);
+
+				const int32 BytesWritten = SSL_write(DTLSContext->GetSSLPtr(), Packet.GetData(), Packet.GetNumBytes());
+				if (BytesWritten != Packet.GetNumBytes())
+				{
+					UE_LOG(LogDTLSHandler, Warning, TEXT("Failed to write entire outgoing packet to SSL: %d / %d"), BytesWritten, Packet.GetNumBytes());
+					Packet.SetError();
+					return;
+				}
+
+				int32 Pending = BIO_ctrl_pending(DTLSContext->GetFilterBIO());
+				if (Pending > 0)
+				{
+					check(Pending <= sizeof(TempBuffer));
+
+					const int32 BytesRead = BIO_read(DTLSContext->GetOutBIO(), TempBuffer, Pending);
+					if (ensure(BytesRead == Pending))
+					{
+						NewPacket.SerializeBits(TempBuffer, BytesRead * 8);
+					}
+					else
+					{
+						UE_LOG(LogDTLSHandler, Error, TEXT("BIO_read error: %d"), BytesRead);
+						Packet.SetError();
+					}
+				}
+				else
+				{
+					UE_LOG(LogDTLSHandler, Error, TEXT("BIO_ctrl_pending error: %d"), Pending);
+					Packet.SetError();
+				}
+			}
+			else if (InternalState == EDTLSHandlerState::Handshaking)
+			{
+				// not a warning as it is expected that this could happen during handshaking
+				UE_LOG(LogDTLSHandler, Log, TEXT("Attempted to send packet during handshaking, dropping."));
+				Packet.Reset();
+				return;
+			}
+			else
+			{
+				UE_LOG(LogDTLSHandler, Error, TEXT("Attempted to send packet while handler was in invalid state"));
+				Packet.SetError();
+			}
+		}
+		else
+		{
+			UE_LOG(LogDTLSHandler, Error, TEXT("Invalid DTLS context."));
+			Packet.SetError();
+		}
+	}
+	else
+	{
+		NewPacket.SerializeBits(Packet.GetData(), Packet.GetNumBits());
+	}
+
+	if (!Packet.IsError())
+	{
+		Packet = MoveTemp(NewPacket);
+	}
+}
+
+void FDTLSHandlerComponent::Tick(float DeltaTime)
+{
+	UpdateHandshake();
+}
+
+int32 FDTLSHandlerComponent::GetReservedPacketBits() const
+{
+	const int32 BlockSizeInBytes = 32;
+
+	// Worst case includes the encryption bit, handshake bit, the termination bit, padding up to the next whole byte, a block of padding, and the header size
+	return 3 + 7 + (BlockSizeInBytes * 8) + FMath::Max(DTLS1_RT_HEADER_LENGTH, DTLS1_HM_HEADER_LENGTH) * 8;
+}
+
+void FDTLSHandlerComponent::CountBytes(FArchive& Ar) const
+{
+	FEncryptionComponent::CountBytes(Ar);
+
+	const SIZE_T SizeOfThis = sizeof(*this) - sizeof(FEncryptionComponent);
+	Ar.CountBytes(SizeOfThis, SizeOfThis);
+}
\ No newline at end of file
diff --git a/Engine/Plugins/Runtime/PacketHandlers/DTLSHandlerComponent/Source/Public/DTLSCertStore.h b/Engine/Plugins/Runtime/PacketHandlers/DTLSHandlerComponent/Source/Public/DTLSCertStore.h
new file mode 100644
--- /dev/null
+++ b/Engine/Plugins/Runtime/PacketHandlers/DTLSHandlerComponent/Source/Public/DTLSCertStore.h
@@ -0,0 +1,59 @@
+// Copyright 1998-2019 Epic Games, Inc. All Rights Reserved.
+
+#pragma once
+
+#include "DTLSHandlerTypes.h"
+#include "DTLSCertificate.h"
+
+/*
+* Certificate store that can generate self-signed X509 certificates for DTLS
+*/
+struct DTLSHANDLERCOMPONENT_API FDTLSCertStore
+{
+public:
+	/**
+	 * Retrieve cert store
+	 *
+	 * @return singleton cert store reference
+	 */
+	static FDTLSCertStore& Get();
+
+	/**
+	 * Create a new certificate
+	 *
+	 * @return shared pointer to certificate, valid if creation succeeded
+	 */
+	TSharedPtr<FDTLSCertificate> CreateCert();
+
+	/**
+	 * Create a new certificate and store internally
+	 *
+	 * @Param Identifier name to use when storing certificate for later use
+	 *
+	 * @return shared pointer to certificate, valid if creation succeeded
+	 */
+	TSharedPtr<FDTLSCertificate> CreateCert(const FString& Identifier);
+
+	/**
+	 * Retrieve a certificate using unique identifier
+	 *
+	 * @Param Identifier unique identifier used when cert was added
+	 *
+	 * @return shared pointer to certificate, valid if it was found
+	 */
+	TSharedPtr<FDTLSCertificate> GetCert(const FString& Identifier) const;
+
+	/**
+	 * Remove a certificate using unique identifier
+	 *
+	 * @Param Identifier unique identifier used when cert was added
+	 *
+	 * @return true if a certificate was removed
+	 */
+	bool RemoveCert(const FString& Identifier);
+
+private:
+	TMap<FString, TSharedPtr<FDTLSCertificate>> CertMap;
+
+	static TUniquePtr<FDTLSCertStore> Instance;
+};
diff --git a/Engine/Plugins/Runtime/PacketHandlers/DTLSHandlerComponent/Source/Public/DTLSCertificate.h b/Engine/Plugins/Runtime/PacketHandlers/DTLSHandlerComponent/Source/Public/DTLSCertificate.h
new file mode 100644
--- /dev/null
+++ b/Engine/Plugins/Runtime/PacketHandlers/DTLSHandlerComponent/Source/Public/DTLSCertificate.h
@@ -0,0 +1,70 @@
+// Copyright 1998-2019 Epic Games, Inc. All Rights Reserved.
+
+#pragma once
+
+#include "DTLSHandlerTypes.h"
+
+/*
+* Wrapper for a fingerprint (SHA256 hash) of an X509 certificate
+*/
+struct FDTLSFingerprint
+{
+public:
+	/** SHA256 hash length in bytes */
+	static constexpr uint32 Length = 32;
+
+	FDTLSFingerprint()
+	{
+		Reset();
+	}
+
+	FDTLSFingerprint(const FDTLSFingerprint&) = delete;
+	FDTLSFingerprint& operator=(const FDTLSFingerprint&) = delete;
+
+	/** Zero the fingerprint */
+	void Reset()
+	{
+		FMemory::Memzero(Data, Length);
+	}
+	
+	/** Get array view of fingerprint data */
+	TArrayView<const uint8> GetData() const { return MakeArrayView(Data, FDTLSFingerprint::Length); }
+
+	uint8 Data[Length];
+};
+
+/*
+* Container for an X509 certificate
+*/
+struct FDTLSCertificate
+{
+public:
+	FDTLSCertificate();
+	~FDTLSCertificate();
+
+	FDTLSCertificate(const FDTLSCertificate&) = delete;
+	FDTLSCertificate& operator=(const FDTLSCertificate&) = delete;
+
+	/** Get OpenSSL private key pointer */
+	EVP_PKEY* GetPKey() const { return PKey; }
+
+	/** Get OpenSSL X509 certificate pointer */
+	X509* GetCertificate() const { return Certificate; }
+
+	/** Get array view of fingerprint data */
+	TArrayView<const uint8> GetFingerprint() const { return Fingerprint.GetData(); }
+
+	/**
+	 * Generate a self-signed certificate
+	 *
+	 * @return true if creation succeeded
+	 */
+	bool GenerateCertificate();
+
+private:
+	void FreeCertificate();
+
+	EVP_PKEY* PKey;
+	X509* Certificate;
+	FDTLSFingerprint Fingerprint;
+};
\ No newline at end of file
diff --git a/Engine/Plugins/Runtime/PacketHandlers/DTLSHandlerComponent/Source/Public/DTLSContext.h b/Engine/Plugins/Runtime/PacketHandlers/DTLSHandlerComponent/Source/Public/DTLSContext.h
new file mode 100644
--- /dev/null
+++ b/Engine/Plugins/Runtime/PacketHandlers/DTLSHandlerComponent/Source/Public/DTLSContext.h
@@ -0,0 +1,120 @@
+// Copyright 1998-2019 Epic Games, Inc. All Rights Reserved.
+
+#pragma once
+
+#include "DTLSHandlerTypes.h"
+#include "DTLSCertificate.h"
+
+class FDTLSHandlerComponent;
+
+/*
+* Wrapper for a pre-shared key value to be used instead of self-signed certificates
+*/
+struct FDTLSPreSharedKey
+{
+public:
+	FDTLSPreSharedKey() {}
+
+	FDTLSPreSharedKey(const FDTLSPreSharedKey&) = delete;
+	FDTLSPreSharedKey& operator=(const FDTLSPreSharedKey&) = delete;
+
+	TArrayView<const uint8> GetKey() const 
+	{ 
+		return MakeArrayView(Key); 
+	}
+	
+	const FString& GetIdentity() const 
+	{ 
+		return Identity; 
+	}
+	
+	void SetPreSharedKey(TArrayView<const uint8> KeyView)
+	{
+		Key.Append(KeyView.begin(), KeyView.Num());
+	}
+
+	void SetIdentity(const FString& InIdentity)
+	{
+		Identity = InIdentity;
+	}
+
+private:
+	TArray<uint8> Key;
+	FString Identity;
+};
+
+/*
+* OpenSSL BIO applied to SSL connection to enforce MTU for handshake packets
+*/
+struct FMTUFilter
+{
+public:
+	static int Create(BIO* BIOPtr);
+	static int Destroy(BIO* BIOPtr);
+	static int Write(BIO* BIOPtr, const char* Data, int DataLen);
+	static long Ctrl(BIO* BIOPtr, int CtrlCommand, long, void*);
+
+public:
+	TArray<int32, TInlineAllocator<16>> FragmentSizes;
+};
+
+/*
+* Context data for OpenSSL connection setup
+*/
+struct FDTLSContext
+{
+	FDTLSContext() = delete;
+
+public:
+	FDTLSContext(EDTLSContextType InContextType);
+	~FDTLSContext();
+
+	/**
+	 * Initialize the context
+	 *
+	 * @param MaxPacketSize the MTU to respect when processing handshake packets
+	 * @param CertId unique identifier of certificate to use, retrieved from FDTLSCertStore
+	 * @param Handler the packet handler component that owns this context
+	 *
+	 * @return false if initialization process failed for any reason
+	 */
+	bool Initialize(const int32 MaxPacketSize, const FString& CertId, FDTLSHandlerComponent* Handler);
+
+	// Get OpenSSL SSL pointer
+	SSL* GetSSLPtr() const { return SSLPtr; }
+
+	// Get OpenSSL input BIO pointer
+	BIO* GetInBIO() const { return InBIO; }
+
+	// Get OpenSSL output BIO pointer
+	BIO* GetOutBIO() const { return OutBIO; }
+
+	// Get OpenSSL custom filter BIO pointer
+	BIO* GetFilterBIO() const { return FilterBIO; }
+
+	/**
+	 * Query for context type
+	 *
+	 * @return EDTLSContextType context type
+	 */
+	EDTLSContextType GetContextType() const { return ContextType; }
+
+	/**
+	 * Query for handshaking status
+	 *
+	 * @return true if the DTLS handshake has completed
+	 */
+	bool IsHandshakeComplete() const;
+
+private:
+	EDTLSContextType ContextType;
+
+	SSL_CTX* SSLContext;
+	SSL* SSLPtr;
+	BIO* InBIO;
+	BIO* OutBIO;
+	BIO* FilterBIO;
+	BIO_METHOD* MTUFilterMethod;
+
+	TSharedPtr<FDTLSCertificate> Cert;
+};
diff --git a/Engine/Plugins/Runtime/PacketHandlers/DTLSHandlerComponent/Source/Public/DTLSHandlerComponent.h b/Engine/Plugins/Runtime/PacketHandlers/DTLSHandlerComponent/Source/Public/DTLSHandlerComponent.h
new file mode 100644
--- /dev/null
+++ b/Engine/Plugins/Runtime/PacketHandlers/DTLSHandlerComponent/Source/Public/DTLSHandlerComponent.h
@@ -0,0 +1,84 @@
+// Copyright 1998-2019 Epic Games, Inc. All Rights Reserved.
+
+#pragma once
+
+#include "CoreMinimal.h"
+#include "PacketHandler.h"
+#include "EncryptionComponent.h"
+#include "UObject/CoreNet.h"
+#include "DTLSContext.h"
+
+extern TAutoConsoleVariable<int32> CVarPreSharedKeys;
+
+/*
+* DTLS encryption component.
+*/
+class DTLSHANDLERCOMPONENT_API FDTLSHandlerComponent : public FEncryptionComponent
+{
+public:
+	FDTLSHandlerComponent();
+
+	virtual void SetEncryptionData(const FEncryptionData& EncryptionData) override;
+
+	// After calling this, future outgoing packets will be encrypted (until a call to DisableEncryption).
+	virtual void EnableEncryption() override;
+
+	// After calling this, future outgoing packets will not be encrypted (until a call to DisableEncryption).
+	virtual void DisableEncryption() override;
+
+	// Returns true if encryption is currently enabled.
+	virtual bool IsEncryptionEnabled() const override;
+
+	// HandlerComponent interface
+	virtual void Initialize() override;
+	virtual bool IsValid() const override;
+	virtual void Incoming(FBitReader& Packet) override;
+	virtual void Outgoing(FBitWriter& Packet, FOutPacketTraits& Traits) override;
+	virtual void IncomingConnectionless(const TSharedPtr<const FInternetAddr>& Address, FBitReader& Packet) override {}
+	virtual void OutgoingConnectionless(const TSharedPtr<const FInternetAddr>& Address, FBitWriter& Packet, FOutPacketTraits& Traits) override {}
+	virtual int32 GetReservedPacketBits() const override;
+	virtual void CountBytes(FArchive& Ar) const override;
+
+	virtual void Tick(float DeltaTime) override;
+
+	enum class EDTLSHandlerState
+	{
+		Unencrypted,
+		Handshaking,
+		Encrypted,
+	};
+
+	// Retrieve pre shared key, if set
+	const FDTLSPreSharedKey* GetPreSharedKey() const { return PreSharedKey.Get(); }
+
+	// Retrieve expected remote certificate fingerprint, if set
+	const FDTLSFingerprint* GetRemoteFingerprint() const { return RemoteFingerprint.Get(); }
+
+private:
+	// Process DTLS handshake
+	void UpdateHandshake();
+
+private:
+	EDTLSHandlerState InternalState;
+
+	TUniquePtr<FDTLSContext> DTLSContext;
+	TUniquePtr<FDTLSPreSharedKey> PreSharedKey;
+	TUniquePtr<FDTLSFingerprint> RemoteFingerprint;
+
+	FString CertId;
+
+	uint8 TempBuffer[MAX_PACKET_SIZE];
+};
+
+/**
+ * The public interface to this module.
+ */
+class FDTLSHandlerComponentModule : public FPacketHandlerComponentModuleInterface
+{
+public:
+	virtual void StartupModule() override;
+	virtual void ShutdownModule() override;
+
+	/* Creates an instance of this component */
+	virtual TSharedPtr<HandlerComponent> CreateComponentInstance(FString& Options) override;
+};
diff --git a/Engine/Plugins/Runtime/PacketHandlers/DTLSHandlerComponent/Source/Public/DTLSHandlerTypes.h b/Engine/Plugins/Runtime/PacketHandlers/DTLSHandlerComponent/Source/Public/DTLSHandlerTypes.h
new file mode 100644
--- /dev/null
+++ b/Engine/Plugins/Runtime/PacketHandlers/DTLSHandlerComponent/Source/Public/DTLSHandlerTypes.h
@@ -0,0 +1,34 @@
+// Copyright 1998-2019 Epic Games, Inc. All Rights Reserved.
+
+#pragma once
+
+#include "CoreMinimal.h"
+
+DECLARE_LOG_CATEGORY_EXTERN(LogDTLSHandler, Log, All);
+
+/** Forward some OpenSSL types  */
+
+struct evp_pkey_st;
+typedef struct evp_pkey_st EVP_PKEY;
+
+struct x509_st;
+typedef struct x509_st X509;
+
+struct ssl_ctx_st;
+typedef struct ssl_ctx_st SSL_CTX;
+
+struct ssl_st;
+typedef struct ssl_st SSL;
+
+struct bio_st;
+typedef struct bio_st BIO;
+
+typedef struct bio_method_st BIO_METHOD;
+
+enum class EDTLSContextType
+{
+	Server,
+	Client
+};
+
+const TCHAR* LexToString(EDTLSContextType ContextType);
\ No newline at end of file
