CIS fix

#rnx

diff --git a/Engine/Plugins/Runtime/PacketHandlers/DTLSHandlerComponent/Source/Private/DTLSCertStore.cpp b/Engine/Plugins/Runtime/PacketHandlers/DTLSHandlerComponent/Source/Private/DTLSCertStore.cpp
--- a/Engine/Plugins/Runtime/PacketHandlers/DTLSHandlerComponent/Source/Private/DTLSCertStore.cpp
+++ b/Engine/Plugins/Runtime/PacketHandlers/DTLSHandlerComponent/Source/Private/DTLSCertStore.cpp
@@ -1,10 +1,11 @@
 // Copyright 1998-2019 Epic Games, Inc. All Rights Reserved.
 
 #include "DTLSCertStore.h"
+#include "Misc/FileHelper.h"
 
 #if !UE_BUILD_SHIPPING
 static TAutoConsoleVariable<int32> CVarDTLSDebugFingerprints(TEXT("DTLS.DebugFingerprints"), 1, TEXT(""));
-#endif
+#endif 
 
 TUniquePtr<FDTLSCertStore> FDTLSCertStore::Instance;
 
