non-unity fix

#rnx

diff --git a/Engine/Plugins/Runtime/PacketHandlers/DTLSHandlerComponent/Source/Private/DTLSCertStore.cpp b/Engine/Plugins/Runtime/PacketHandlers/DTLSHandlerComponent/Source/Private/DTLSCertStore.cpp
--- a/Engine/Plugins/Runtime/PacketHandlers/DTLSHandlerComponent/Source/Private/DTLSCertStore.cpp
+++ b/Engine/Plugins/Runtime/PacketHandlers/DTLSHandlerComponent/Source/Private/DTLSCertStore.cpp
@@ -2,10 +2,11 @@
 
 #include "DTLSCertStore.h"
 #include "Misc/FileHelper.h"
+#include "HAL/IConsoleManager.h"
 
 #if !UE_BUILD_SHIPPING
 static TAutoConsoleVariable<int32> CVarDTLSDebugFingerprints(TEXT("DTLS.DebugFingerprints"), 1, TEXT(""));
-#endif 
+#endif
 
 TUniquePtr<FDTLSCertStore> FDTLSCertStore::Instance;
 
