CIS fix dejavu

#rnx

diff --git a/Engine/Plugins/Runtime/PacketHandlers/DTLSHandlerComponent/Source/Private/DTLSCertStore.cpp b/Engine/Plugins/Runtime/PacketHandlers/DTLSHandlerComponent/Source/Private/DTLSCertStore.cpp
--- a/Engine/Plugins/Runtime/PacketHandlers/DTLSHandlerComponent/Source/Private/DTLSCertStore.cpp
+++ b/Engine/Plugins/Runtime/PacketHandlers/DTLSHandlerComponent/Source/Private/DTLSCertStore.cpp
@@ -5,7 +5,7 @@
 #include "HAL/IConsoleManager.h"
 
 #if !UE_BUILD_SHIPPING
-static TAutoConsoleVariable<int32> CVarDTLSDebugFingerprints(TEXT("DTLS.DebugFingerprints"), 1, TEXT(""));
+static TAutoConsoleVariable<int32> CVarDTLSDebugFingerprints(TEXT("DTLS.DebugFingerprints"), 0, TEXT(""));
 #endif
 
 TUniquePtr<FDTLSCertStore> FDTLSCertStore::Instance;
@@ -40,7 +40,7 @@ TSharedPtr<FDTLSCertificate> FDTLSCertStore::CreateCert(const FString& Identifie
 		CertMap.Emplace(Identifier, Cert);
 
 #if !UE_BUILD_SHIPPING
-		const bool bDebugFingerprints = CVarDTLSDebugFingerprints.GetValueOnAnyThread();
+		const bool bDebugFingerprints = (CVarDTLSDebugFingerprints.GetValueOnAnyThread() != 0);
 
 		if (bDebugFingerprints)
 		{
