Adding option to use DTLS packet handler to Lyra
Adding import/export of test certs to PEM
Adjust HTTP dependencies when it is supported but not required

diff --git a/Engine/Plugins/Media/ElectraUtil/Source/ElectraHTTPStream/ElectraHTTPStream.Build.cs b/Engine/Plugins/Media/ElectraUtil/Source/ElectraHTTPStream/ElectraHTTPStream.Build.cs
--- a/Engine/Plugins/Media/ElectraUtil/Source/ElectraHTTPStream/ElectraHTTPStream.Build.cs
+++ b/Engine/Plugins/Media/ElectraUtil/Source/ElectraHTTPStream/ElectraHTTPStream.Build.cs
@@ -86,8 +86,7 @@ public class ElectraHTTPStream : ModuleRules
 
 		PrivateDependencyModuleNames.AddRange(
 			new string[] {
-				"ElectraBase",
-				"SSL",
+				"ElectraBase"
 			});
 
 		// When Curl is used we need to make sure there is no conflict with its use in the engine's HTTP module
@@ -140,11 +139,16 @@ public class ElectraHTTPStream : ModuleRules
 		}
 
 		PublicDefinitions.Add("ELECTRA_HTTPSTREAM_REQUIRES_SECURE_CONNECTIONS=" + (bPlatformRequiresSecureConnections ? "1" : "0"));
-
+		
 		if (bPlatformRequiresOpenSSL)
 		{
+			PrivateDependencyModuleNames.Add("SSL");
 			AddEngineThirdPartyPrivateStaticDependencies(Target, "OpenSSL");
 		}
+		else
+		{
+			PrivateDefinitions.Add("WITH_SSL=0");
+		}
 
 		// Apple
 		if (Target.Platform == UnrealTargetPlatform.IOS || Target.Platform == UnrealTargetPlatform.TVOS || Target.Platform == UnrealTargetPlatform.Mac)
diff --git a/Engine/Plugins/Runtime/PacketHandlers/DTLSHandlerComponent/Source/Private/DTLSCertStore.cpp b/Engine/Plugins/Runtime/PacketHandlers/DTLSHandlerComponent/Source/Private/DTLSCertStore.cpp
--- a/Engine/Plugins/Runtime/PacketHandlers/DTLSHandlerComponent/Source/Private/DTLSCertStore.cpp
+++ b/Engine/Plugins/Runtime/PacketHandlers/DTLSHandlerComponent/Source/Private/DTLSCertStore.cpp
@@ -60,6 +60,32 @@ TSharedPtr<FDTLSCertificate> FDTLSCertStore::GetCert(const FString& Identifier)
 	return CertMap.FindRef(Identifier);
 }
 
+TSharedPtr<FDTLSCertificate> FDTLSCertStore::ImportCert(const FString& CertPath) const
+{
+	TSharedRef<FDTLSCertificate> Cert = MakeShared<FDTLSCertificate>();
+	if (Cert->ImportCertificate(CertPath))
+	{
+		return Cert;
+	}
+
+	UE_LOG(LogDTLSHandler, Error, TEXT("ImportCert: Failed to import certificate"));
+	return nullptr;
+}
+
+TSharedPtr<FDTLSCertificate> FDTLSCertStore::ImportCert(const FString& CertPath, const FString& Identifier)
+{
+	TSharedRef<FDTLSCertificate> Cert = MakeShared<FDTLSCertificate>();
+	if (Cert->ImportCertificate(CertPath) && !Identifier.IsEmpty())
+	{
+		CertMap.Emplace(Identifier, Cert);
+
+		return Cert;
+	}
+
+	UE_LOG(LogDTLSHandler, Error, TEXT("ImportCert: Failed to import certificate"));
+	return nullptr;
+}
+
 bool FDTLSCertStore::RemoveCert(const FString& Identifier)
 {
 	return (CertMap.Remove(Identifier) != 0);
diff --git a/Engine/Plugins/Runtime/PacketHandlers/DTLSHandlerComponent/Source/Private/DTLSCertificate.cpp b/Engine/Plugins/Runtime/PacketHandlers/DTLSHandlerComponent/Source/Private/DTLSCertificate.cpp
--- a/Engine/Plugins/Runtime/PacketHandlers/DTLSHandlerComponent/Source/Private/DTLSCertificate.cpp
+++ b/Engine/Plugins/Runtime/PacketHandlers/DTLSHandlerComponent/Source/Private/DTLSCertificate.cpp
@@ -3,6 +3,7 @@
 #include "DTLSCertificate.h"
 #include "CoreGlobals.h"
 #include "Misc/ScopeExit.h"
+#include "HAL/FileManager.h"
 
 #define UI UI_ST
 THIRD_PARTY_INCLUDES_START
@@ -11,6 +12,7 @@ THIRD_PARTY_INCLUDES_START
 #include <openssl/bn.h>
 #include <openssl/evp.h>
 #include <openssl/ossl_typ.h>
+#include <openssl/pem.h>
 #include <openssl/rsa.h>
 THIRD_PARTY_INCLUDES_END
 #undef UI
@@ -47,6 +49,97 @@ void FDTLSCertificate::FreeCertificate()
 	Fingerprint.Reset();
 }
 
+bool FDTLSCertificate::ExportCertificate(const FString& CertPath)
+{
+	if (Certificate)
+	{
+		if (!CertPath.IsEmpty())
+		{
+			BIO* CertificateBio = BIO_new(BIO_s_mem());
+			PEM_write_bio_X509(CertificateBio, Certificate);
+
+			TArray<uint8> MemBuffer;
+			uint64 MemBufferSize = BIO_number_written(CertificateBio);
+			MemBuffer.AddUninitialized(MemBufferSize);
+
+			BIO_read(CertificateBio, MemBuffer.GetData(), MemBufferSize);
+			BIO_free(CertificateBio);
+			CertificateBio = nullptr;
+
+			TUniquePtr<FArchive> ExportAr(IFileManager::Get().CreateFileWriter(*CertPath));
+
+			if (ExportAr.IsValid())
+			{
+				ExportAr->Serialize(MemBuffer.GetData(), MemBuffer.Num());
+				return true;
+			}
+			else
+			{
+				UE_LOG(LogDTLSHandler, Error, TEXT("ExportCertificate: Unable to create file writer for: %s"), *CertPath);
+			}
+		}
+		else
+		{
+			UE_LOG(LogDTLSHandler, Error, TEXT("ExportCertificate: Empty cert path"));
+		}
+	}
+	else
+	{
+		UE_LOG(LogDTLSHandler, Error, TEXT("ExportCertificate: Cert not valid"));
+	}
+
+	return false;
+}
+
+bool FDTLSCertificate::ImportCertificate(const FString& CertPath)
+{
+	if (Certificate == nullptr)
+	{
+		if (!CertPath.IsEmpty())
+		{
+			TUniquePtr<FArchive> ImportAr(IFileManager::Get().CreateFileReader(*CertPath));
+			if (ImportAr.IsValid())
+			{
+				const int64 MemBufferSize = ImportAr->TotalSize();
+
+				TArray<uint8> MemBuffer;
+				MemBuffer.AddUninitialized(MemBufferSize + 1);
+				ImportAr->Serialize(MemBuffer.GetData(), MemBufferSize);
+				MemBuffer[MemBufferSize] = '\0';
+
+				BIO* CertificateBio = BIO_new_mem_buf(MemBuffer.GetData(), -1);
+				Certificate = PEM_read_bio_X509(CertificateBio, nullptr, 0, nullptr);
+
+				BIO_free(CertificateBio);
+				CertificateBio = nullptr;
+
+				if (Certificate)
+				{
+					return GenerateFingerprint();
+				}
+				else
+				{
+					UE_LOG(LogDTLSHandler, Error, TEXT("ImportCertificate: Unable to read valid x509 certifcate"));
+				}
+			}
+			else
+			{
+				UE_LOG(LogDTLSHandler, Error, TEXT("ImportCertificate: Unable to create file reader for: %s"), *CertPath);
+			}
+		}
+		else
+		{
+			UE_LOG(LogDTLSHandler, Error, TEXT("ImportCertificate: Empty cert path"));
+		}
+	}
+	else
+	{
+		UE_LOG(LogDTLSHandler, Error, TEXT("ImportCertificate: Cert already exists"));
+	}
+
+	return false;
+}
+
 bool FDTLSCertificate::GenerateCertificate(const FTimespan& Lifetime)
 {
 	bool bSuccess = false;
@@ -120,18 +213,23 @@ bool FDTLSCertificate::GenerateCertificate(const FTimespan& Lifetime)
 		return bSuccess;
 	}
 
+	bSuccess = GenerateFingerprint();
+
+	return bSuccess;
+}
+
+bool FDTLSCertificate::GenerateFingerprint()
+{
 	Fingerprint.Reset();
 
 	uint32 HashLen = FDTLSFingerprint::Length;
 
-	ResultCode = X509_digest(Certificate, EVP_sha256(), Fingerprint.Data, &HashLen);
+	int32 ResultCode = X509_digest(Certificate, EVP_sha256(), Fingerprint.Data, &HashLen);
 	if (ResultCode <= 0)
 	{
 		UE_LOG(LogDTLSHandler, Error, TEXT("Failed to hash cert: %d"), ResultCode);
-		return bSuccess;
+		return false;
 	}
 
-	bSuccess = true;
-
-	return bSuccess;
-}
+	return true;
+}
\ No newline at end of file
diff --git a/Engine/Plugins/Runtime/PacketHandlers/DTLSHandlerComponent/Source/Public/DTLSCertStore.h b/Engine/Plugins/Runtime/PacketHandlers/DTLSHandlerComponent/Source/Public/DTLSCertStore.h
--- a/Engine/Plugins/Runtime/PacketHandlers/DTLSHandlerComponent/Source/Public/DTLSCertStore.h
+++ b/Engine/Plugins/Runtime/PacketHandlers/DTLSHandlerComponent/Source/Public/DTLSCertStore.h
@@ -46,6 +46,25 @@ public:
 	 */
 	TSharedPtr<FDTLSCertificate> GetCert(const FString& Identifier) const;
 
+	/**
+	 * Import a certificate from file
+	 *
+	 * @Param CertPath path to certificate file
+	 *
+	 * @return shared pointer to certificate, valid if import succeeded
+	 */
+	TSharedPtr<FDTLSCertificate> ImportCert(const FString& CertPath) const;
+
+	/**
+	 * Import a certificate from file and store internally
+	 *
+	 * @Param CertPath path to certificate file
+	 * @Param Identifier name to use when storing certificate for later use
+	 *
+	 * @return shared pointer to certificate, valid if import succeeded
+	 */
+	TSharedPtr<FDTLSCertificate> ImportCert(const FString& CertPath, const FString& Identifier);
+
 	/**
 	 * Remove a certificate using unique identifier
 	 *
diff --git a/Engine/Plugins/Runtime/PacketHandlers/DTLSHandlerComponent/Source/Public/DTLSCertificate.h b/Engine/Plugins/Runtime/PacketHandlers/DTLSHandlerComponent/Source/Public/DTLSCertificate.h
--- a/Engine/Plugins/Runtime/PacketHandlers/DTLSHandlerComponent/Source/Public/DTLSCertificate.h
+++ b/Engine/Plugins/Runtime/PacketHandlers/DTLSHandlerComponent/Source/Public/DTLSCertificate.h
@@ -37,7 +37,7 @@ public:
 /*
 * Container for an X509 certificate
 */
-struct FDTLSCertificate
+struct DTLSHANDLERCOMPONENT_API FDTLSCertificate
 {
 public:
 	FDTLSCertificate();
@@ -45,6 +45,8 @@ public:
 
 	FDTLSCertificate(const FDTLSCertificate&) = delete;
 	FDTLSCertificate& operator=(const FDTLSCertificate&) = delete;
+	FDTLSCertificate(FDTLSCertificate&&) = delete;
+	FDTLSCertificate& operator=(FDTLSCertificate&&) = delete;
 
 	/** Get OpenSSL private key pointer */
 	EVP_PKEY* GetPKey() const { return PKey; }
@@ -63,8 +65,25 @@ public:
 	 */
 	bool GenerateCertificate(const FTimespan& Lifetime);
 
+	/**
+	 * Export current certificate to PEM file format
+	 *
+	 * @param CertPath path to output file
+	 * @return true if export succeeded
+	 */
+	bool ExportCertificate(const FString& CertPath);
+
+	/**
+	 * Import certificate from PEM file format
+	 *
+	 * @param CertPath path to input file
+	 * @return true if import succeeded
+	 */
+	bool ImportCertificate(const FString& CertPath);
+
 private:
 	void FreeCertificate();
+	bool GenerateFingerprint();
 
 	EVP_PKEY* PKey;
 	X509* Certificate;
diff --git a/Engine/Source/Runtime/Online/HTTP/HTTP.Build.cs b/Engine/Source/Runtime/Online/HTTP/HTTP.Build.cs
--- a/Engine/Source/Runtime/Online/HTTP/HTTP.Build.cs
+++ b/Engine/Source/Runtime/Online/HTTP/HTTP.Build.cs
@@ -45,12 +45,6 @@ public class HTTP : ModuleRules
 			}
 			);
 
-		PrivateDependencyModuleNames.AddRange(
-			new string[] {
-				"SSL",
-			}
-			);
-
 		if (bPlatformSupportsCurl)
 		{
 			PrivateDependencyModuleNames.AddRange(
@@ -88,8 +82,13 @@ public class HTTP : ModuleRules
 
 		if (bPlatformRequiresOpenSSL)
 		{
+			PrivateDependencyModuleNames.Add("SSL");
 			AddEngineThirdPartyPrivateStaticDependencies(Target, "OpenSSL");
 		}
+		else
+		{
+			PrivateDefinitions.Add("WITH_SSL=0");
+		}
 
 		if (Target.Platform == UnrealTargetPlatform.IOS || Target.Platform == UnrealTargetPlatform.TVOS || Target.Platform == UnrealTargetPlatform.Mac)
 		{
