Converting checks to proper error handling

diff --git a/Engine/Plugins/Runtime/PacketHandlers/DTLSHandlerComponent/Source/Private/DTLSHandlerComponent.cpp b/Engine/Plugins/Runtime/PacketHandlers/DTLSHandlerComponent/Source/Private/DTLSHandlerComponent.cpp
--- a/Engine/Plugins/Runtime/PacketHandlers/DTLSHandlerComponent/Source/Private/DTLSHandlerComponent.cpp
+++ b/Engine/Plugins/Runtime/PacketHandlers/DTLSHandlerComponent/Source/Private/DTLSHandlerComponent.cpp
@@ -327,28 +327,41 @@ void FDTLSHandlerComponent::DoHandshake()
 	int32 Pending = BIO_ctrl_pending(DTLSContext->GetFilterBIO());
 	while (Pending > 0)
 	{
-		check(Pending <= sizeof(TempBuffer));
-
-		const int32 BytesRead = BIO_read(DTLSContext->GetOutBIO(), TempBuffer, Pending);
-		if (BytesRead > 0)
+		if (Pending <= sizeof(TempBuffer))
 		{
-			check(BytesRead == Pending);
-			check(BytesRead <= MAX_PACKET_SIZE);
-
-			FBitWriter OutPacket(0, true);
-			OutPacket.WriteBit(1);	// encryption enabled
-			OutPacket.WriteBit(1);	// handshake packet
-			OutPacket.SerializeBits(TempBuffer, BytesRead * 8);
+			const int32 BytesRead = BIO_read(DTLSContext->GetOutBIO(), TempBuffer, Pending);
+			if (BytesRead > 0)
+			{
+				if (BytesRead == Pending)
+				{
+					FBitWriter OutPacket(0, true);
+					OutPacket.WriteBit(1);	// encryption enabled
+					OutPacket.WriteBit(1);	// handshake packet
+					OutPacket.SerializeBits(TempBuffer, BytesRead * 8);
 
-			UE_LOG(LogDTLSHandler, Verbose, TEXT("DoHandshake:  Sending handshake packet: %d bytes"), BytesRead);
+					UE_LOG(LogDTLSHandler, Verbose, TEXT("DoHandshake:  Sending handshake packet: %d bytes"), BytesRead);
 
-			// SendHandlerPacket is a low level send and is not reliable
-			FOutPacketTraits Traits;
-			Handler->SendHandlerPacket(this, OutPacket, Traits);
+					// SendHandlerPacket is a low level send and is not reliable
+					FOutPacketTraits Traits;
+					Handler->SendHandlerPacket(this, OutPacket, Traits);
+				}
+				else
+				{
+					DTLSContext.Reset();
+					UE_LOG(LogDTLSHandler, Error, TEXT("Unexpected value from BIO_read: Read: %d Pending: %d"), BytesRead, Pending);
+					return;
+				}
+			}
+			else
+			{
+				UE_LOG(LogDTLSHandler, Error, TEXT("BIO_read error: %d"), BytesRead);
+				return;
+			}
 		}
 		else
 		{
-			UE_LOG(LogDTLSHandler, Error, TEXT("BIO_read error: %d"), BytesRead);
+			DTLSContext.Reset();
+			UE_LOG(LogDTLSHandler, Error, TEXT("BIO_ctrl_pending returned value out of range: %d"), Pending);
 			return;
 		}
 
@@ -416,16 +429,22 @@ void FDTLSHandlerComponent::Outgoing(FBitWriter& Packet, FOutPacketTraits& Trait
 				int32 Pending = BIO_ctrl_pending(DTLSContext->GetFilterBIO());
 				if (Pending > 0)
 				{
-					check(Pending <= sizeof(TempBuffer));
-
-					const int32 BytesRead = BIO_read(DTLSContext->GetOutBIO(), TempBuffer, Pending);
-					if (ensure(BytesRead == Pending))
+					if (Pending <= sizeof(TempBuffer))
 					{
-						NewPacket.SerializeBits(TempBuffer, BytesRead * 8);
+						const int32 BytesRead = BIO_read(DTLSContext->GetOutBIO(), TempBuffer, Pending);
+						if (ensure(BytesRead == Pending))
+						{
+							NewPacket.SerializeBits(TempBuffer, BytesRead * 8);
+						}
+						else
+						{
+							UE_LOG(LogDTLSHandler, Error, TEXT("BIO_read error: %d"), BytesRead);
+							Packet.SetError();
+						}
 					}
 					else
 					{
-						UE_LOG(LogDTLSHandler, Error, TEXT("BIO_read error: %d"), BytesRead);
+						UE_LOG(LogDTLSHandler, Error, TEXT("BIO_ctrl_pending returned value out of range: %d"), Pending);
 						Packet.SetError();
 					}
 				}
