Fixed incorrect FString constructor.

diff --git a/Engine/Plugins/Runtime/PacketHandlers/DTLSHandlerComponent/Source/Private/DTLSHandlerComponent.cpp b/Engine/Plugins/Runtime/PacketHandlers/DTLSHandlerComponent/Source/Private/DTLSHandlerComponent.cpp
--- a/Engine/Plugins/Runtime/PacketHandlers/DTLSHandlerComponent/Source/Private/DTLSHandlerComponent.cpp
+++ b/Engine/Plugins/Runtime/PacketHandlers/DTLSHandlerComponent/Source/Private/DTLSHandlerComponent.cpp
@@ -531,7 +531,7 @@ void FDTLSHandlerComponent::LogError(const TCHAR* Context, int32 Result)
 		char* ErrorBuffer = nullptr;
 		size_t BufferLength = BIO_get_mem_data(MemBIO, &ErrorBuffer);
 
-		FString ErrorString(ErrorBuffer, BufferLength);
+		FString ErrorString = FString::ConstructFromPtrSize(ErrorBuffer, BufferLength);
 
 		BIO_free(MemBIO);
 
